Add a standalone program to locally execute XGE task lists, allowing MegaXGE-style builds on the local machine.

diff --git a/Engine/Source/Programs/ParallelExecutor/ParallelExecutor.Build.cs b/Engine/Source/Programs/ParallelExecutor/ParallelExecutor.Build.cs
new file mode 100644
--- /dev/null
+++ b/Engine/Source/Programs/ParallelExecutor/ParallelExecutor.Build.cs
@@ -0,0 +1,14 @@
+﻿// Copyright 1998-2015 Epic Games, Inc. All Rights Reserved.
+
+using UnrealBuildTool;
+
+public class ParallelExecutor : ModuleRules
+{
+	public ParallelExecutor(TargetInfo Target)
+	{
+		PublicIncludePaths.Add("Runtime/Launch/Public");
+		PublicIncludePaths.Add("Runtime/Launch/Private"); // Yuck. Required for RequiredProgramMainCPPInclude.h. (Also yuck).
+
+		PrivateDependencyModuleNames.AddRange(new string[] { "Core", "Projects", "XmlParser" });
+	}
+}
diff --git a/Engine/Source/Programs/ParallelExecutor/ParallelExecutor.Target.cs b/Engine/Source/Programs/ParallelExecutor/ParallelExecutor.Target.cs
new file mode 100644
--- /dev/null
+++ b/Engine/Source/Programs/ParallelExecutor/ParallelExecutor.Target.cs
@@ -0,0 +1,54 @@
+﻿// Copyright 1998-2015 Epic Games, Inc. All Rights Reserved.
+
+using UnrealBuildTool;
+using System.Collections.Generic;
+
+public class ParallelExecutorTarget : TargetRules
+{
+	public ParallelExecutorTarget(TargetInfo Target)
+	{
+		Type = TargetType.Program;
+	}
+
+	//
+	// TargetRules interface.
+	//
+
+	public override void SetupBinaries(
+		TargetInfo Target,
+		ref List<UEBuildBinaryConfiguration> OutBuildBinaryConfigurations,
+		ref List<string> OutExtraModuleNames
+		)
+	{
+		OutBuildBinaryConfigurations.Add(
+			new UEBuildBinaryConfiguration( InType: UEBuildBinaryType.Executable, InModuleNames: new List<string>() { "ParallelExecutor" })
+			);
+	}
+
+	public override bool ShouldCompileMonolithic(UnrealTargetPlatform InPlatform, UnrealTargetConfiguration InConfiguration)
+	{
+		return true;
+	}
+
+	public override void SetupGlobalEnvironment(
+		TargetInfo Target,
+		ref LinkEnvironmentConfiguration OutLinkEnvironmentConfiguration,
+		ref CPPEnvironmentConfiguration OutCPPEnvironmentConfiguration
+		)
+	{
+		UEBuildConfiguration.bCompileLeanAndMeanUE = true;
+		BuildConfiguration.bUseMallocProfiler = false;
+
+		// No editor needed
+		UEBuildConfiguration.bCompileICU = false;
+		UEBuildConfiguration.bBuildEditor = false;
+		UEBuildConfiguration.bBuildWithEditorOnlyData = false;
+
+		// Currently this app is not linking against the engine, so we'll compile out references from Core to the rest of the engine
+		UEBuildConfiguration.bCompileAgainstEngine = false;
+		UEBuildConfiguration.bCompileAgainstCoreUObject = false;
+
+		// UnrealHeaderTool is a console application, not a Windows app (sets entry point to main(), instead of WinMain())
+		OutLinkEnvironmentConfiguration.bIsBuildingConsoleApplication = true;
+	}
+}
diff --git a/Engine/Source/Programs/ParallelExecutor/Private/BuildGraph.cpp b/Engine/Source/Programs/ParallelExecutor/Private/BuildGraph.cpp
new file mode 100644
--- /dev/null
+++ b/Engine/Source/Programs/ParallelExecutor/Private/BuildGraph.cpp
@@ -0,0 +1,473 @@
+// Copyright 1998-2015 Epic Games, Inc. All Rights Reserved.
+
+#include "ParallelExecutor.h"
+#include "BuildGraph.h"
+#include "XmlFile.h"
+
+FBuildAction::FBuildAction()
+{
+	SortIndex = INDEX_NONE;
+	bSkipIfProjectFailed = false;
+	MissingDependencyCount = 0;
+	TotalDependants = 0;
+}
+
+FBuildAction::~FBuildAction()
+{
+}
+
+//////////////////////////////////////////////////////////////////////////////////////////////////////////////////////////////////
+
+FBuildActionExecutor::FBuildActionExecutor(FBuildAction* InAction, FCriticalSection& InCriticalSection, FEvent* InCompletedEvent, TArray<FBuildActionExecutor*>& InCompletedActions, const TMap<FString, FString>& InEnvironment, void* InJobObject) : 
+	Action(InAction),
+	CriticalSection(InCriticalSection),
+	CompletedEvent(InCompletedEvent),
+	CompletedActions(InCompletedActions),
+	ExitCode(-1),
+	Environment(InEnvironment),
+	JobObject(InJobObject)
+{
+}
+
+uint32 FBuildActionExecutor::Run()
+{
+	ExitCode = ExecuteAction();
+
+	CriticalSection.Lock();
+	CompletedActions.Add(this);
+	CriticalSection.Unlock();
+
+	CompletedEvent->Trigger();
+	return 0;
+}
+
+int32 FBuildActionExecutor::ExecuteAction()
+{
+	int32 ExitCode = -1;
+
+	void* ReadPipe;
+	void* WritePipe;
+	if (!FPlatformProcess::CreatePipe(ReadPipe, WritePipe))
+	{
+		Log.Add(TEXT("Failed to create pipes"));
+	}
+	else
+	{
+		// Create the child process
+		FProcHandle ProcessHandle = CreateChildProcess(Action->ToolPath, ExpandEnvironmentVariables(Action->ToolArguments), Action->WorkingDirectory, Environment, WritePipe, JobObject);
+		if(!ProcessHandle.IsValid())
+		{
+			Log.Add(FString::Printf(TEXT("Couldn't start processs %s"), *Action->ToolPath));
+		}
+		else
+		{
+			FString BufferedText;
+			for(;;)
+			{
+				WaitForOutput(ProcessHandle, ReadPipe);
+				BufferedText += FPlatformProcess::ReadPipe(ReadPipe);
+
+				int32 NewlineIdx;
+				while(BufferedText.FindChar(TEXT('\n'), NewlineIdx))
+				{
+					Log.Add(BufferedText.Left(NewlineIdx).Replace(TEXT("\r"), TEXT("")));
+					BufferedText = BufferedText.Mid(NewlineIdx + 1);
+				}
+
+				if(!FPlatformProcess::IsProcRunning(ProcessHandle))
+				{
+					break;
+				}
+			}
+			FPlatformProcess::GetProcReturnCode(ProcessHandle, &ExitCode);
+		}
+		FPlatformProcess::ClosePipe(ReadPipe, WritePipe);
+	}
+
+	return ExitCode;
+}
+
+FString FBuildActionExecutor::ExpandEnvironmentVariables(const FString& Input)
+{
+	FString Output = Input;
+	for(int StartIndex = 0;;)
+	{
+		int32 Index = Output.Find(TEXT("$("), ESearchCase::CaseSensitive, ESearchDir::FromStart, StartIndex);
+		if(Index == INDEX_NONE) break;
+
+		int32 EndIndex = Output.Find(TEXT(")"), ESearchCase::CaseSensitive, ESearchDir::FromStart, Index);
+		if(EndIndex == INDEX_NONE) break;
+
+		FString VariableName = Output.Mid(Index + 2, EndIndex - Index - 2);
+
+		TArray<TCHAR> RawVariableValue;
+		RawVariableValue.AddZeroed(2048);
+		FPlatformMisc::GetEnvironmentVariable(*VariableName, RawVariableValue.GetData(), RawVariableValue.Num() - 1);
+
+		FString VariableValue = RawVariableValue.GetData();
+
+		Output.RemoveAt(Index, EndIndex + 1 - Index);
+		Output.InsertAt(Index, VariableValue);
+
+		StartIndex = Index + VariableValue.Len();
+	}
+	return Output;
+}
+
+FProcHandle FBuildActionExecutor::CreateChildProcess(const FString& FileName, const FString& Arguments, const FString& WorkingDir, const TMap<FString, FString>& Environment, void* WritePipe, void* JobObject)
+{
+	#if PLATFORM_WINDOWS
+		SECURITY_ATTRIBUTES Attr;
+		ZeroMemory(&Attr, sizeof(Attr));
+		Attr.nLength = sizeof(SECURITY_ATTRIBUTES);
+		Attr.bInheritHandle = true;
+
+		STARTUPINFO StartupInfo;
+		ZeroMemory(&StartupInfo, sizeof(StartupInfo));
+		StartupInfo.cb = sizeof(StartupInfo);
+		StartupInfo.dwFlags = STARTF_USESTDHANDLES;
+		StartupInfo.wShowWindow = SW_HIDE;
+		StartupInfo.hStdInput = ::GetStdHandle(/*STD_INPUT_HANDLE*/((uint32)-10));
+		StartupInfo.hStdOutput = WritePipe;
+		StartupInfo.hStdError = WritePipe;
+
+		FString CommandLine = FString::Printf(TEXT("\"%s\" %s"), *FileName, *Arguments);
+
+		// Get the current environment
+		TCHAR* CurrentEnvironmentBlock = GetEnvironmentStrings();
+
+		// Create the new environment block
+		TArray<TCHAR> NewEnvironmentBlock;
+		for(const TCHAR* Variable = CurrentEnvironmentBlock; *Variable != 0; Variable = FCString::Strchr(Variable, 0) + 1)
+		{
+			const TCHAR* Equals = FCString::Strchr(Variable, TEXT('='));
+			if(Equals != nullptr && !Environment.Contains(FString(Equals - Variable, Variable)))
+			{
+				NewEnvironmentBlock.Append(Variable, FCString::Strlen(Variable) + 1);
+			}
+		}
+		for(const TPair<FString, FString>& Pair : Environment)
+		{
+			if(Pair.Value.Len() > 0)
+			{
+				FString Assignment = FString::Printf(TEXT("%s=%s"), *Pair.Key, *Pair.Value);
+				NewEnvironmentBlock.Append(*Assignment, Assignment.Len() + 1);
+			}
+		}
+		NewEnvironmentBlock.Add(0);
+
+		// Free the current environment block
+		FreeEnvironmentStrings(CurrentEnvironmentBlock);
+
+		PROCESS_INFORMATION ProcessInfo;
+		if (!CreateProcess(NULL, CommandLine.GetCharArray().GetData(), &Attr, &Attr, true, BELOW_NORMAL_PRIORITY_CLASS | /*DETACHED_PROCESS | */CREATE_UNICODE_ENVIRONMENT | CREATE_SUSPENDED, NewEnvironmentBlock.GetData(), *WorkingDir, &StartupInfo, &ProcessInfo))
+		{
+			return FProcHandle(0);
+		}
+
+		if(JobObject != nullptr)
+		{
+			AssignProcessToJobObject(JobObject, ProcessInfo.hProcess);
+		}
+
+		ResumeThread(ProcessInfo.hThread);
+		::CloseHandle(ProcessInfo.hThread);
+
+		return FProcHandle(ProcessInfo.hProcess);
+	#else
+		#error CreateProcessWithEnvironment is not implemented for this platform
+	#endif
+}
+
+void FBuildActionExecutor::WaitForOutput(FProcHandle ProcessHandle, void* ReadPipe)
+{
+	#if PLATFORM_WINDOWS
+		HANDLE Handles[] = { ProcessHandle.Get(), ReadPipe };
+		::WaitForMultipleObjects(ARRAY_COUNT(Handles), Handles, false, INFINITE);
+	#else
+		#error WaitForOutput is not implemented for this platform
+	#endif
+}
+
+//////////////////////////////////////////////////////////////////////////////////////////////////////////////////////////////////
+
+FBuildGraph::FBuildGraph()
+{
+}
+
+FBuildGraph::~FBuildGraph()
+{
+}
+
+bool FBuildGraph::ReadFromFile(const FString& InputPath)
+{
+	FXmlFile File;
+	if(!File.LoadFile(InputPath))
+	{
+		return false;
+	}
+
+	const FXmlNode* RootNode = File.GetRootNode();
+	if(RootNode == nullptr || RootNode->GetTag() != TEXT("BuildSet"))
+	{
+		return false;
+	}
+
+	const FXmlNode* EnvironmentsNode = RootNode->FindChildNode(TEXT("Environments"));
+	if(EnvironmentsNode == nullptr)
+	{
+		return false;
+	}
+
+	// Read the tool environment
+	TMap<FString, const FXmlNode*> NameToTool;
+	for(const FXmlNode* EnvironmentNode = EnvironmentsNode->GetFirstChildNode(); EnvironmentNode != nullptr; EnvironmentNode = EnvironmentNode->GetNextNode())
+	{
+		if(EnvironmentNode->GetTag() == TEXT("Environment"))
+		{
+			// Read the tools nodes
+			const FXmlNode* ToolsNode = EnvironmentNode->FindChildNode(TEXT("Tools"));
+			if(ToolsNode != nullptr)
+			{
+				for(const FXmlNode* ToolNode = ToolsNode->GetFirstChildNode(); ToolNode != nullptr; ToolNode = ToolNode->GetNextNode())
+				{
+					if(ToolNode->GetTag() == TEXT("Tool"))
+					{
+						FString Name = ToolNode->GetAttribute(TEXT("Name"));
+						if(Name.Len() > 0)
+						{
+							NameToTool.Add(Name, ToolNode);
+						}
+					}
+				}
+			}
+
+			// Read the environment variables
+			const FXmlNode* VariablesNode = EnvironmentNode->FindChildNode(TEXT("Variables"));
+			if(VariablesNode != nullptr)
+			{
+				for(const FXmlNode* VariableNode = VariablesNode->GetFirstChildNode(); VariableNode != nullptr; VariableNode = VariableNode->GetNextNode())
+				{
+					if(VariableNode->GetTag() == TEXT("Variable"))
+					{
+						FString Key = VariableNode->GetAttribute(TEXT("Name"));
+						if(Key.Len() > 0)
+						{
+							Variables.FindOrAdd(Key) = VariableNode->GetAttribute(TEXT("Value"));
+						}
+					}
+				}
+			}
+		}
+	}
+
+	// Read all the tasks for each project, and convert them into actions
+	TSet<FBuildAction*> VisitedActions;
+	for(const FXmlNode* ProjectNode = RootNode->GetFirstChildNode(); ProjectNode != nullptr; ProjectNode = ProjectNode->GetNextNode())
+	{
+		if(ProjectNode->GetTag() == TEXT("Project"))
+		{
+			int32 SortIndex = 0;
+			TMap<FString, FBuildAction*> NameToAction;
+			for(const FXmlNode* TaskNode = ProjectNode->GetFirstChildNode(); TaskNode != nullptr; TaskNode = TaskNode->GetNextNode())
+			{
+				if(TaskNode->GetTag() == TEXT("Task"))
+				{
+					const FXmlNode** ToolNode = NameToTool.Find(TaskNode->GetAttribute(TEXT("Tool")));
+					if(ToolNode != nullptr)
+					{
+						FBuildAction* Action = FindOrAddAction(NameToAction, TaskNode->GetAttribute(TEXT("Name")));
+						Action->SortIndex = ++SortIndex;
+						Action->Caption = TaskNode->GetAttribute(TEXT("Caption"));
+						Action->Prefix = (*ToolNode)->GetAttribute(TEXT("GroupPrefix"));
+						Action->ToolPath = (*ToolNode)->GetAttribute(TEXT("Path"));
+						Action->ToolArguments = (*ToolNode)->GetAttribute(TEXT("Params"));
+						Action->WorkingDirectory = TaskNode->GetAttribute(TEXT("WorkingDir"));
+						Action->bSkipIfProjectFailed = TaskNode->GetAttribute(TEXT("SkipIfProjectFailed")) == TEXT("True");
+
+						TArray<FString> DependsOnList;
+						TaskNode->GetAttribute(TEXT("DependsOn")).ParseIntoArray(DependsOnList, TEXT(";"));
+
+						for(const FString& DependsOn: DependsOnList)
+						{
+							FBuildAction* DependsOnAction = FindOrAddAction(NameToAction, DependsOn);
+							Action->Dependencies.AddUnique(DependsOnAction);
+							DependsOnAction->Dependants.AddUnique(Action);
+						}
+
+						VisitedActions.Empty();
+						RecursiveIncDependants(Action, VisitedActions);
+						Action->MissingDependencyCount = DependsOnList.Num();
+					}
+				}
+			}
+		}
+	}
+
+	// Make sure there have been no actions added which were referenced but never declared
+	for(const FBuildAction* Action : Actions)
+	{
+		if(Action->SortIndex == INDEX_NONE)
+		{
+			return false;
+		}
+	}
+
+	return true;
+}
+
+#if PLATFORM_WINDOWS
+void* JobObjectToClose;
+BOOL WINAPI QuitHandlerRoutine(::DWORD CtrlType)
+{
+	::CloseHandle(JobObjectToClose);
+	TerminateProcess(GetCurrentProcess(), 999);
+	return 1;
+}
+#endif
+
+int32 FBuildGraph::ExecuteInParallel(int32 MaxProcesses)
+{
+	wprintf(TEXT("Building with %d %s...\n"), MaxProcesses, (MaxProcesses == 1)? TEXT("process") : TEXT("processes"));
+
+	// Create the list of things to process
+	TArray<FBuildAction*> QueuedActions;
+	for(FBuildAction* Action: Actions)
+	{
+		if(Action->MissingDependencyCount == 0)
+		{
+			QueuedActions.Add(Action);
+		}
+	}
+
+	// Create a job object for all the child processes
+	void* JobObject = CreateJobObject();
+
+#if PLATFORM_WINDOWS
+	JobObjectToClose = JobObject;
+	::SetConsoleCtrlHandler(&QuitHandlerRoutine, 1);
+#endif
+
+	// Process all the input queue
+	int32 ExitCode = 0;
+	FString CurrentPrefix;
+	TMap<FBuildActionExecutor*, FRunnableThread*> ExecutingActions;
+	TArray<FBuildActionExecutor*> CompletedActions;
+
+	FCriticalSection CriticalSection;
+
+	FEvent* CompletedEvent = FPlatformProcess::CreateSynchEvent(false);
+	while(QueuedActions.Num() > 0 || ExecutingActions.Num() > 0)
+	{
+		// Sort the actions by the number of things dependent on them
+		QueuedActions.Sort([](const FBuildAction& A, const FBuildAction& B){ return (A.TotalDependants < B.TotalDependants) || (A.TotalDependants == B.TotalDependants && A.SortIndex > B.SortIndex); });
+
+		// Create threads up to the maximum number of actions
+		while(ExecutingActions.Num() < MaxProcesses && QueuedActions.Num() > 0)
+		{
+			FBuildAction* Action = QueuedActions.Pop();
+			FBuildActionExecutor* ExecutingAction = new FBuildActionExecutor(Action, CriticalSection, CompletedEvent, CompletedActions, Variables, JobObject);
+			FRunnableThread* ExecutingThread = FRunnableThread::Create(ExecutingAction, *FString::Printf(TEXT("Build:%s"), *Action->Caption));
+			ExecutingActions.Add(ExecutingAction, ExecutingThread);
+		}
+
+		// Wait for something to finish
+		CompletedEvent->Wait();
+
+		// Wait for something to finish and flush it to the log
+		CriticalSection.Lock();
+		for(FBuildActionExecutor* CompletedAction : CompletedActions)
+		{
+			// Join the thread
+			FRunnableThread* Thread = ExecutingActions[CompletedAction];
+			Thread->WaitForCompletion();
+			ExecutingActions.Remove(CompletedAction);
+
+			// Write it to the log
+			if(CompletedAction->Log.Num() > 0)
+			{
+				if(CurrentPrefix != CompletedAction->Action->Prefix)
+				{
+					CurrentPrefix = CompletedAction->Action->Prefix;
+					wprintf(TEXT("%s\n"), *CurrentPrefix);
+				}
+				for(const FString& Line : CompletedAction->Log)
+				{
+					wprintf(TEXT("%s\n"), *Line);
+				}
+			}
+
+			// Check the exit code
+			if(CompletedAction->ExitCode != 0)
+			{
+				ExitCode = CompletedAction->ExitCode;
+				if(!CompletedAction->Action->bSkipIfProjectFailed)
+				{
+					for(TPair<FBuildActionExecutor*, FRunnableThread*> Pair : ExecutingActions)
+					{
+						Pair.Value->Kill(false);
+						delete Pair.Key;
+					}
+					FPlatformProcess::ReturnSynchEventToPool(CompletedEvent);
+					return ExitCode;
+				}
+			}
+
+			// Mark all the dependants as done
+			for(FBuildAction* DependantAction : CompletedAction->Action->Dependants)
+			{
+				if(--DependantAction->MissingDependencyCount == 0)
+				{
+					QueuedActions.Add(DependantAction);
+				}
+			}
+
+			// Destroy the action
+			delete CompletedAction;
+		}
+		CompletedActions.Reset();
+		CriticalSection.Unlock();
+	}
+	FPlatformProcess::ReturnSynchEventToPool(CompletedEvent);
+	return ExitCode;
+}
+
+FBuildAction* FBuildGraph::FindOrAddAction(TMap<FString, FBuildAction*>& NameToAction, const FString& Name)
+{
+	FBuildAction*& Action = NameToAction.FindOrAdd(Name);
+	if(Action == nullptr)
+	{
+		Action = new FBuildAction();
+		Actions.Add(Action);
+	}
+	return Action;
+}
+
+void FBuildGraph::RecursiveIncDependants(FBuildAction* Action, TSet<FBuildAction*>& VisitedActions)
+{
+	for(FBuildAction* Dependency : Action->Dependencies)
+	{
+		if(!VisitedActions.Contains(Action))
+		{
+			VisitedActions.Add(Action);
+			Dependency->TotalDependants++;
+			RecursiveIncDependants(Dependency, VisitedActions);
+		}
+	}
+}
+
+void* FBuildGraph::CreateJobObject()
+{
+	void* JobObject = nullptr;
+	#if PLATFORM_WINDOWS
+		JobObject = ::CreateJobObject(nullptr, nullptr);
+		if(JobObject != nullptr)
+		{
+			JOBOBJECT_EXTENDED_LIMIT_INFORMATION LimitInformation;
+			ZeroMemory(&LimitInformation, sizeof(LimitInformation));
+			LimitInformation.BasicLimitInformation.LimitFlags = JOB_OBJECT_LIMIT_KILL_ON_JOB_CLOSE;
+			SetInformationJobObject(JobObject, JobObjectExtendedLimitInformation, &LimitInformation, sizeof(LimitInformation));
+		}
+	#endif
+	return JobObject;
+}
diff --git a/Engine/Source/Programs/ParallelExecutor/Private/BuildGraph.h b/Engine/Source/Programs/ParallelExecutor/Private/BuildGraph.h
new file mode 100644
--- /dev/null
+++ b/Engine/Source/Programs/ParallelExecutor/Private/BuildGraph.h
@@ -0,0 +1,66 @@
+// Copyright 1998-2015 Epic Games, Inc. All Rights Reserved.
+
+#pragma once
+
+class FBuildAction
+{
+public:
+	int32 SortIndex;
+
+	FString Caption;
+	FString Prefix;
+	FString ToolPath;
+	FString ToolArguments;
+	FString WorkingDirectory;
+	bool bSkipIfProjectFailed;
+
+	TArray<FBuildAction*> Dependants;
+	TArray<FBuildAction*> Dependencies;
+
+	int32 TotalDependants;
+	int32 MissingDependencyCount;
+
+	FBuildAction();
+	~FBuildAction();
+};
+
+class FBuildActionExecutor : public FRunnable
+{
+public:
+	FBuildAction* Action;
+	TArray<FString> Log;
+	int32 ExitCode;
+
+	FBuildActionExecutor(FBuildAction* InAction, FCriticalSection& InCriticalSection, FEvent* InCompletedEvent, TArray<FBuildActionExecutor*>& InCompletedActions, const TMap<FString, FString>& InEnvironment, void* InJobObject);
+	virtual uint32 Run() override;
+
+private:
+	FCriticalSection& CriticalSection;
+	FEvent* CompletedEvent;
+	TArray<FBuildActionExecutor*>& CompletedActions;
+	const TMap<FString, FString>& Environment;
+	void* JobObject;
+
+	int32 ExecuteAction();
+	static FString ExpandEnvironmentVariables(const FString& Input);
+	static FProcHandle CreateChildProcess(const FString& FileName, const FString& Arguments, const FString& WorkingDir, const TMap<FString, FString>& Environment, void* WritePipe, void* JobObject);
+	static void WaitForOutput(FProcHandle ProcessHandle, void* ReadPipe);
+};
+
+class FBuildGraph
+{
+public:
+	FBuildGraph();
+	~FBuildGraph();
+	bool ReadFromFile(const FString& InputPath);
+	int32 ExecuteInParallel(int32 MaxProcesses);
+
+private:
+	TArray<FBuildAction*> Actions;
+	TMap<FString, FString> Variables;
+
+	FBuildAction* FindOrAddAction(TMap<FString, FBuildAction*>& NameToAction, const FString& Name);
+	static void RecursiveIncDependants(FBuildAction* Action, TSet<FBuildAction*>& VisitedActions);
+
+	static void* CreateJobObject();
+};
diff --git a/Engine/Source/Programs/ParallelExecutor/Private/ParallelExecutor.cpp b/Engine/Source/Programs/ParallelExecutor/Private/ParallelExecutor.cpp
new file mode 100644
--- /dev/null
+++ b/Engine/Source/Programs/ParallelExecutor/Private/ParallelExecutor.cpp
@@ -0,0 +1,32 @@
+// Copyright 1998-2015 Epic Games, Inc. All Rights Reserved.
+
+#include "ParallelExecutor.h"
+#include "RequiredProgramMainCPPInclude.h"
+#include "BuildGraph.h"
+
+IMPLEMENT_APPLICATION(ParallelExecutor, "ParallelExecutor")
+
+
+int main(int ArgC, const char* ArgV[])
+{
+	// TODO: command line options for limiting number of processes
+	// TODO: add ability to modify environment for child process, and add to job
+	// TODO: add a FPlatformMisc::WaitForProcessOutput() with optional parameter for cancel event
+
+	FCommandLine::Set(TEXT(""));
+
+	if(ArgC != 2)
+	{
+		wprintf(TEXT("Missing argument to ParallelExecutor for build graph."));
+		return 1;
+	}
+
+	FBuildGraph Graph;
+	if(!Graph.ReadFromFile(ArgV[1]))
+	{
+		wprintf(TEXT("Couldn't read '%s'"), ArgV[1]);
+		return 1;
+	}
+
+	return Graph.ExecuteInParallel(1);//FPlatformMisc::NumberOfCoresIncludingHyperthreads());
+}
diff --git a/Engine/Source/Programs/ParallelExecutor/Private/ParallelExecutor.h b/Engine/Source/Programs/ParallelExecutor/Private/ParallelExecutor.h
new file mode 100644
--- /dev/null
+++ b/Engine/Source/Programs/ParallelExecutor/Private/ParallelExecutor.h
@@ -0,0 +1,5 @@
+// Copyright 1998-2015 Epic Games, Inc. All Rights Reserved.
+
+#pragma once
+
+#include "Core.h"
