Fix number of processes used by ParallelExecutor, and don't wait on the pipe output from child processes (it returns immediately, so just timeout on waiting for the process handle instead).

diff --git a/Engine/Source/Programs/ParallelExecutor/Private/BuildGraph.cpp b/Engine/Source/Programs/ParallelExecutor/Private/BuildGraph.cpp
--- a/Engine/Source/Programs/ParallelExecutor/Private/BuildGraph.cpp
+++ b/Engine/Source/Programs/ParallelExecutor/Private/BuildGraph.cpp
@@ -182,8 +182,7 @@ FProcHandle FBuildActionExecutor::CreateChildProcess(const FString& FileName, co
 void FBuildActionExecutor::WaitForOutput(FProcHandle ProcessHandle, void* ReadPipe)
 {
 	#if PLATFORM_WINDOWS
-		HANDLE Handles[] = { ProcessHandle.Get(), ReadPipe };
-		::WaitForMultipleObjects(ARRAY_COUNT(Handles), Handles, false, INFINITE);
+		::WaitForSingleObject(ProcessHandle.Get(), 500);
 	#else
 		#error WaitForOutput is not implemented for this platform
 	#endif
diff --git a/Engine/Source/Programs/ParallelExecutor/Private/ParallelExecutor.cpp b/Engine/Source/Programs/ParallelExecutor/Private/ParallelExecutor.cpp
--- a/Engine/Source/Programs/ParallelExecutor/Private/ParallelExecutor.cpp
+++ b/Engine/Source/Programs/ParallelExecutor/Private/ParallelExecutor.cpp
@@ -28,5 +28,5 @@ int main(int ArgC, const char* ArgV[])
 		return 1;
 	}
 
-	return Graph.ExecuteInParallel(1);//FPlatformMisc::NumberOfCoresIncludingHyperthreads());
+	return Graph.ExecuteInParallel(FPlatformMisc::NumberOfCoresIncludingHyperthreads());
 }
