replaced GroupPrefix with OutputPrefix in ParallelExecutor's output to more closely match the output of XGE

diff --git a/Engine/Source/Programs/ParallelExecutor/Private/BuildGraph.cpp b/Engine/Source/Programs/ParallelExecutor/Private/BuildGraph.cpp
--- a/Engine/Source/Programs/ParallelExecutor/Private/BuildGraph.cpp
+++ b/Engine/Source/Programs/ParallelExecutor/Private/BuildGraph.cpp
@@ -270,7 +270,7 @@ bool FBuildGraph::ReadFromFile(const FString& InputPath)
 						FBuildAction* Action = FindOrAddAction(NameToAction, TaskNode->GetAttribute(TEXT("Name")));
 						Action->SortIndex = ++SortIndex;
 						Action->Caption = TaskNode->GetAttribute(TEXT("Caption"));
-						Action->Prefix = (*ToolNode)->GetAttribute(TEXT("GroupPrefix"));
+						Action->Prefix = (*ToolNode)->GetAttribute(TEXT("OutputPrefix"));
 						Action->ToolPath = (*ToolNode)->GetAttribute(TEXT("Path"));
 						Action->ToolArguments = (*ToolNode)->GetAttribute(TEXT("Params"));
 						Action->WorkingDirectory = TaskNode->GetAttribute(TEXT("WorkingDir"));
