Add GroupPrefix back into ParallelExecutor output (so OutputPrefix is in addition to, instead of replacing, GroupPrefix)

diff --git a/Engine/Source/Programs/ParallelExecutor/Private/BuildGraph.cpp b/Engine/Source/Programs/ParallelExecutor/Private/BuildGraph.cpp
--- a/Engine/Source/Programs/ParallelExecutor/Private/BuildGraph.cpp
+++ b/Engine/Source/Programs/ParallelExecutor/Private/BuildGraph.cpp
@@ -270,7 +270,8 @@ bool FBuildGraph::ReadFromFile(const FString& InputPath)
 						FBuildAction* Action = FindOrAddAction(NameToAction, TaskNode->GetAttribute(TEXT("Name")));
 						Action->SortIndex = ++SortIndex;
 						Action->Caption = TaskNode->GetAttribute(TEXT("Caption"));
-						Action->Prefix = (*ToolNode)->GetAttribute(TEXT("OutputPrefix"));
+						Action->GroupPrefix = (*ToolNode)->GetAttribute(TEXT("GroupPrefix"));
+						Action->OutputPrefix = (*ToolNode)->GetAttribute(TEXT("OutputPrefix"));
 						Action->ToolPath = (*ToolNode)->GetAttribute(TEXT("Path"));
 						Action->ToolArguments = (*ToolNode)->GetAttribute(TEXT("Params"));
 						Action->WorkingDirectory = TaskNode->GetAttribute(TEXT("WorkingDir"));
@@ -395,11 +396,15 @@ int32 FBuildGraph::ExecuteInParallel(int32 MaxProcesses)
 			// Write it to the log
 			if(CompletedAction->Log.Num() > 0)
 			{
-				if(CurrentPrefix != CompletedAction->Action->Prefix)
+				if(CurrentPrefix != CompletedAction->Action->GroupPrefix)
 				{
-					CurrentPrefix = CompletedAction->Action->Prefix;
+					CurrentPrefix = CompletedAction->Action->GroupPrefix;
 					wprintf(TEXT("%s\n"), *CurrentPrefix);
 				}
+				if (!CompletedAction->Action->OutputPrefix.IsEmpty())
+				{
+					wprintf(TEXT("%s\n"), *CompletedAction->Action->OutputPrefix);
+				}
 				for(const FString& Line : CompletedAction->Log)
 				{
 					wprintf(TEXT("%s\n"), *Line);
diff --git a/Engine/Source/Programs/ParallelExecutor/Private/BuildGraph.h b/Engine/Source/Programs/ParallelExecutor/Private/BuildGraph.h
--- a/Engine/Source/Programs/ParallelExecutor/Private/BuildGraph.h
+++ b/Engine/Source/Programs/ParallelExecutor/Private/BuildGraph.h
@@ -8,7 +8,8 @@ public:
 	int32 SortIndex;
 
 	FString Caption;
-	FString Prefix;
+	FString GroupPrefix;
+	FString OutputPrefix;
 	FString ToolPath;
 	FString ToolArguments;
 	FString WorkingDirectory;
