PR #1479: Update ParallelExecutor.cpp (Contributed by [REDACTED_PERSON])

diff --git a/Engine/Source/Programs/ParallelExecutor/Private/ParallelExecutor.cpp b/Engine/Source/Programs/ParallelExecutor/Private/ParallelExecutor.cpp
--- a/Engine/Source/Programs/ParallelExecutor/Private/ParallelExecutor.cpp
+++ b/Engine/Source/Programs/ParallelExecutor/Private/ParallelExecutor.cpp
@@ -26,7 +26,7 @@ int main(int ArgC, const char* ArgV[])
 	FBuildGraph Graph;
 	if(!Graph.ReadFromFile(ArgV[1]))
 	{
-		wprintf(TEXT("Couldn't read '%s'"), ArgV[1]);
+		wprintf(TEXT("Couldn't read '%hs'"), ArgV[1]);
 		return 1;
 	}
 
